#UE4 DOCS: added more tutorials

diff --git a/Engine/Plugins/Editor/IntroTutorials/Source/IntroTutorials/Private/InteractivityData.cpp b/Engine/Plugins/Editor/IntroTutorials/Source/IntroTutorials/Private/InteractivityData.cpp
--- a/Engine/Plugins/Editor/IntroTutorials/Source/IntroTutorials/Private/InteractivityData.cpp
+++ b/Engine/Plugins/Editor/IntroTutorials/Source/IntroTutorials/Private/InteractivityData.cpp
@@ -1772,6 +1772,146 @@ void PopulateInteractivityData(TSharedPtr<FInteractiveTutorials> InteractiveTuto
 	InteractiveTutorials->BindInteractivityData(TEXT("Shared/Tutorials/inEditorTutorialCamera"), FBindInteractivityData::CreateStatic(&FCameraTutorial::BindInteractivityData));
 
 
+	//////////////////////////////////////////////////////////////////////////
+
+
+	struct FInTextureEditorTutorial
+	{
+		static void BindInteractivityData(const FString& InUDNPath, const TSharedRef<class FInteractiveTutorial>& InTutorial)
+		{
+			
+		}
+	};
+
+	InteractiveTutorials->BindInteractivityData(TEXT("Shared/Tutorials/InTextureEditorTutorial"), FBindInteractivityData::CreateStatic(&FInTextureEditorTutorial::BindInteractivityData));
+
+
+	//////////////////////////////////////////////////////////////////////////
+
+
+	struct FInMaterialEditorTutorial
+	{
+		static void BindInteractivityData(const FString& InUDNPath, const TSharedRef<class FInteractiveTutorial>& InTutorial)
+		{
+
+		}
+	};
+
+	InteractiveTutorials->BindInteractivityData(TEXT("Shared/Tutorials/InMaterialEditorTutorial"), FBindInteractivityData::CreateStatic(&FInMaterialEditorTutorial::BindInteractivityData));
+
+
+	//////////////////////////////////////////////////////////////////////////
+
+	
+
+	struct FInParticleSystemEditorTutorial
+	{
+		static void BindInteractivityData(const FString& InUDNPath, const TSharedRef<class FInteractiveTutorial>& InTutorial)
+		{
+
+		}
+	};
+
+	InteractiveTutorials->BindInteractivityData(TEXT("Shared/Tutorials/InParticleSystemEditorTutorial"), FBindInteractivityData::CreateStatic(&FInParticleSystemEditorTutorial::BindInteractivityData));
+
+
+
+	//////////////////////////////////////////////////////////////////////////
+
+
+
+	struct FInPersonaAnimEditorTutorial
+	{
+		static void BindInteractivityData(const FString& InUDNPath, const TSharedRef<class FInteractiveTutorial>& InTutorial)
+		{
+
+		}
+	};
+
+	InteractiveTutorials->BindInteractivityData(TEXT("Shared/Tutorials/InPersonaAnimEditorTutorial"), FBindInteractivityData::CreateStatic(&FInPersonaAnimEditorTutorial::BindInteractivityData));
+
+	
+
+	//////////////////////////////////////////////////////////////////////////
+
+
+
+	struct FInPersonaAnimBlueprintEditorTutorial
+	{
+		static void BindInteractivityData(const FString& InUDNPath, const TSharedRef<class FInteractiveTutorial>& InTutorial)
+		{
+
+		}
+	};
+
+	InteractiveTutorials->BindInteractivityData(TEXT("Shared/Tutorials/InPersonaAnimBlueprintEditorTutorial"), FBindInteractivityData::CreateStatic(&FInPersonaAnimBlueprintEditorTutorial::BindInteractivityData));
+
+
+
+	//////////////////////////////////////////////////////////////////////////
+
+
+
+	struct FInDestructibleMeshEditorTutorial
+	{
+		static void BindInteractivityData(const FString& InUDNPath, const TSharedRef<class FInteractiveTutorial>& InTutorial)
+		{
+
+		}
+	};
+
+	InteractiveTutorials->BindInteractivityData(TEXT("Shared/Tutorials/InDestructibleMeshEditorTutorial"), FBindInteractivityData::CreateStatic(&FInDestructibleMeshEditorTutorial::BindInteractivityData));
+
+
+
+	//////////////////////////////////////////////////////////////////////////
+
+
+
+	struct FInSkeletalMeshEditorTutorial
+	{
+		static void BindInteractivityData(const FString& InUDNPath, const TSharedRef<class FInteractiveTutorial>& InTutorial)
+		{
+
+		}
+	};
+
+	InteractiveTutorials->BindInteractivityData(TEXT("Shared/Tutorials/InSkeletalMeshEditorTutorial"), FBindInteractivityData::CreateStatic(&FInSkeletalMeshEditorTutorial::BindInteractivityData));
+
+
+
+	//////////////////////////////////////////////////////////////////////////
+
+
+
+	struct FInStaticMeshEditorTutorial
+	{
+		static void BindInteractivityData(const FString& InUDNPath, const TSharedRef<class FInteractiveTutorial>& InTutorial)
+		{
+
+		}
+	};
+
+	InteractiveTutorials->BindInteractivityData(TEXT("Shared/Tutorials/InStaticMeshEditorTutorial"), FBindInteractivityData::CreateStatic(&FInStaticMeshEditorTutorial::BindInteractivityData));
+
+
+
+
+	//////////////////////////////////////////////////////////////////////////
+
+
+	struct FInSoundCueEditorTutorial
+	{
+		static void BindInteractivityData(const FString& InUDNPath, const TSharedRef<class FInteractiveTutorial>& InTutorial)
+		{
+
+		}
+	};
+
+	InteractiveTutorials->BindInteractivityData(TEXT("Shared/Tutorials/InSoundCueEditorTutorial"), FBindInteractivityData::CreateStatic(&FInSoundCueEditorTutorial::BindInteractivityData));
+
+
+
 	//////////////////////////////////////////////////////////////////////////
 
 	struct FPlacementTutorial
