Fixing what seems to be coding error.

Contributed by UE4Linux community, part of PR 140.

diff --git a/Engine/Plugins/Editor/IntroTutorials/Source/IntroTutorials/Private/SIntroTutorials.cpp b/Engine/Plugins/Editor/IntroTutorials/Source/IntroTutorials/Private/SIntroTutorials.cpp
--- a/Engine/Plugins/Editor/IntroTutorials/Source/IntroTutorials/Private/SIntroTutorials.cpp
+++ b/Engine/Plugins/Editor/IntroTutorials/Source/IntroTutorials/Private/SIntroTutorials.cpp
@@ -607,7 +607,7 @@ FString SIntroTutorials::GetCurrentExcerptIdentifierName() const
 {
 	if (Excerpts.IsValidIndex(CurrentExcerptIndex))
 	{
-		Excerpts[CurrentExcerptIndex].Name;
+		return Excerpts[CurrentExcerptIndex].Name;
 	}
 
 	return FString();
